Add missing pragma once in new RayTracingTypes.h header file to fix unity builds

diff --git a/Engine/Shaders/Shared/RayTracingTypes.h b/Engine/Shaders/Shared/RayTracingTypes.h
--- a/Engine/Shaders/Shared/RayTracingTypes.h
+++ b/Engine/Shaders/Shared/RayTracingTypes.h
@@ -1,5 +1,7 @@
 // Copyright Epic Games, Inc. All Rights Reserved.
 
+#pragma once
+
 /*================================================================================================
 	RayTracingTypes.h: used in ray tracing shaders and C++ code to define common types
 	!!! Changing this file requires recompilation of the engine !!!
