Bump RAY_TRACING_SHADER_VERSION to force rebuild of any RT shaders that were accidentally cached while using an experimental shader compiler version.

diff --git a/Engine/Shaders/Shared/RayTracingDefinitions.h b/Engine/Shaders/Shared/RayTracingDefinitions.h
--- a/Engine/Shaders/Shared/RayTracingDefinitions.h
+++ b/Engine/Shaders/Shared/RayTracingDefinitions.h
@@ -9,7 +9,7 @@
 
 // Change this to force recompilation of all ray tracing shaders (use https://www.random.org/cgi-bin/randbyte?nbytes=4&format=h)
 // This avoids changing the global ShaderVersion.ush and forcing recompilation of all shaders in the engine (only RT shaders will be affected)
-#define RAY_TRACING_SHADER_VERSION 0x7034a1ef 
+#define RAY_TRACING_SHADER_VERSION 0x7c6a2787
 
 #define RAY_TRACING_REGISTER_SPACE_LOCAL  0 // default register space for hit group (closest hit, any hit, intersection) shader resources
 #define RAY_TRACING_REGISTER_SPACE_GLOBAL 1 // register space for ray generation and miss shaders 
