Reserve a ray mask bit for excluding objects from scene capture views

diff --git a/Engine/Shaders/Shared/RayTracingDefinitions.h b/Engine/Shaders/Shared/RayTracingDefinitions.h
--- a/Engine/Shaders/Shared/RayTracingDefinitions.h
+++ b/Engine/Shaders/Shared/RayTracingDefinitions.h
@@ -19,6 +19,7 @@
 #define RAY_TRACING_MASK_TRANSLUCENT				0x02    // Opaque and alpha tested meshes and particles (e.g. used by translucency tracing pass)
 #define RAY_TRACING_MASK_THIN_SHADOW				0x04    // Whether the thin geometry (e.g. hair) is visible for shadow rays
 #define RAY_TRACING_MASK_SHADOW						0x08    // Whether the geometry is visible for shadow rays
+#define RAY_TRACING_MASK_SCENE_CAPTURE				0x10    // Whether the geometry is visible when rendering scene capture views
 #define RAY_TRACING_MASK_ALL						0xFF
 
 #define RAY_TRACING_SHADER_SLOT_MATERIAL	0
