Path Tracing: Refactor RayTracingTypes.h to avoid code duplication by using the HLSLTypeAliases.h header

diff --git a/Engine/Shaders/Shared/RayTracingTypes.h b/Engine/Shaders/Shared/RayTracingTypes.h
--- a/Engine/Shaders/Shared/RayTracingTypes.h
+++ b/Engine/Shaders/Shared/RayTracingTypes.h
@@ -7,88 +7,36 @@
 	!!! Changing this file requires recompilation of the engine !!!
 =================================================================================================*/
 
-#ifdef __cplusplus
-
-// C++ representation of a light for the path tracer
-// #dxr_todo: Unify this with FRTLightingData ?
-struct FPathTracingLight {
-// 	FVector3f RelativeWorldPosition;
-// 	FVector3f TilePosition;
-	FVector3f TranslatedWorldPosition;
-	FVector3f Normal;
-	FVector3f dPdu;
-	FVector3f dPdv;
-	FVector3f Color;
-	FVector2f Dimensions; // Radius,Length or RectWidth,RectHeight or Sin(Angle/2),0 depending on light type
-	FVector2f Shaping;  // Barndoor controls for RectLights, Cone angles for spots lights
-	float   Attenuation;
-	float   FalloffExponent; // for non-inverse square decay lights only
-	float   VolumetricScatteringIntensity;  // scale for volume contributions
-	int32   IESAtlasIndex;
-	uint32  Flags; // see defines PATHTRACER_FLAG_*
-	uint32  MissShaderIndex;  // used to implement light functions
-	FVector3f TranslatedBoundMin;
-	FVector3f TranslatedBoundMax;
-	uint32 RectLightAtlasUVScale;  // Rect. light atlas UV transformation, encoded as f16x2
-	uint32 RectLightAtlasUVOffset; // Rect. light atlas UV transformation, encoded as f16x2
-	// keep structure aligned
-};
-
-static_assert(sizeof(FPathTracingLight) == 132, "Path tracing light structure should be kept as small as possible");
-
-struct FPathTracingPackedPathState {
-	uint32    RandSeqSampleIndex;
-	uint32    RandSeqSampleSeed;
-	FVector3f Radiance;
-	float     BackgroundVisibility;
-	uint16    Albedo[3];
-	uint16    Normal[3];
-	FVector3f RayOrigin;
-	FVector3f RayDirection;
-	FVector3f PathThroughput;
-	uint16    PathRoughness;
-	uint16    SigmaT[3];
-};
+#include "HLSLStaticAssert.h"
 
-static_assert(sizeof(FPathTracingPackedPathState) == 80, "Packed Path State size should be minimized");
+#ifdef __cplusplus
+#include "HLSLTypeAliases.h"
 
-// C++ representation of a decal for ray tracing
-struct FRayTracingDecal
+namespace UE::HLSL
 {
-	FVector3f TranslatedBoundMin;
-	uint32 Pad0;
-	FVector3f TranslatedBoundMax;
-	uint32 CallableSlotIndex;
-	// keep structure aligned
-};
-
-static_assert(sizeof(FRayTracingDecal) == 32, "Ray tracing decal structure should be aligned to 32 bytes for optimal access on the GPU");
-
-#else
-
-// HLSL side of the structs above
+#endif
 
+// #dxr_todo: Unify this with FRTLightingData ?
 struct FPathTracingLight {
-// 	float3  RelativeWorldPosition;
-// 	float3  TilePosition;
 	float3  TranslatedWorldPosition;
 	float3  Normal;
 	float3  dPdu;
 	float3  dPdv;
 	float3  Color;
-	float2  Dimensions;
-	float2  Shaping;
+	float2  Dimensions; // Radius,Length or RectWidth,RectHeight or Sin(Angle/2),0 depending on light type
+	float2  Shaping;    // Barndoor controls for RectLights, Cone angles for spots lights
 	float   Attenuation;
-	float   FalloffExponent;
+	float   FalloffExponent; // for non-inverse square decay lights only
 	float   VolumetricScatteringIntensity;  // scale for volume contributions
 	int     IESAtlasIndex;
-	uint    Flags;
-	uint    MissShaderIndex;
+	uint    Flags; // see defines PATHTRACER_FLAG_*
+	uint    MissShaderIndex;  // used to implement light functions
 	float3  TranslatedBoundMin;
 	float3  TranslatedBoundMax;
-	uint	RectLightAtlasUVScale;
-	uint	RectLightAtlasUVOffset;
+	uint	RectLightAtlasUVScale;  // Rect. light atlas UV transformation, encoded as f16x2
+	uint	RectLightAtlasUVOffset; // Rect. light atlas UV transformation, encoded as f16x2
 };
+HLSL_STATIC_ASSERT(sizeof(FPathTracingLight) == 132, "Path tracing light structure should be kept as small as possible");
 
 struct FPathTracingPackedPathState {
 	uint      RandSeqSampleIndex;
@@ -101,13 +49,22 @@ struct FPathTracingPackedPathState {
 	float3    PathThroughput;
 	uint2     PackedRoughnessSigma;
 };
+HLSL_STATIC_ASSERT(sizeof(FPathTracingPackedPathState) == 80, "Packed Path State size should be minimized");
 
 struct FRayTracingDecal
 {
 	float3 TranslatedBoundMin;
-	uint Pad0;
+	uint   Pad0; // keep structure aligned
 	float3 TranslatedBoundMax;
-	uint CallableSlotIndex;
+	uint   CallableSlotIndex;
 };
+HLSL_STATIC_ASSERT(sizeof(FRayTracingDecal) == 32, "Ray tracing decal structure should be aligned to 32 bytes for optimal access on the GPU");
+
+#ifdef __cplusplus
+} // namespace UE::HLSL
+
+using FPathTracingLight = UE::HLSL::FPathTracingLight;
+using FPathTracingPackedPathState = UE::HLSL::FPathTracingPackedPathState;
+using FRayTracingDecal = UE::HLSL::FRayTracingDecal;
 
 #endif
