Fix warning about unbalanced scope depth

diff --git a/Engine/Shaders/Shared/LightSceneData.h b/Engine/Shaders/Shared/LightSceneData.h
--- a/Engine/Shaders/Shared/LightSceneData.h
+++ b/Engine/Shaders/Shared/LightSceneData.h
@@ -7,10 +7,12 @@
 
 namespace UE::HLSL
 {
-#else
+#endif
+#ifndef __cplusplus //HLSL
 #include "/Engine/Private/LargeWorldCoordinates.ush"
 #endif
 
+
 /**
  * Has a 1:1 mapping with FLightRenderParameters, but unlike FLightShaderParameters, this is view-independent
  */
