Fix compile errors due to missing file in previous changelist.

diff --git a/Engine/Shaders/Shared/VirtualShadowMapDefinitions.h b/Engine/Shaders/Shared/VirtualShadowMapDefinitions.h
--- a/Engine/Shaders/Shared/VirtualShadowMapDefinitions.h
+++ b/Engine/Shaders/Shared/VirtualShadowMapDefinitions.h
@@ -25,3 +25,32 @@ VirtualShadowMapDefinitions.h: used in virtual shadow map shaders and C++ code t
 
 // Hard limit for max distant lights supported 8k for now - we may revise later. We need to keep them in a fixed range for now to make allocation easy and minimize overhead for indexing.
 #define VSM_MAX_SINGLE_PAGE_SHADOW_MAPS (1024U * 8U)
+
+#ifdef __cplusplus
+#include "HLSLTypeAliases.h"
+
+namespace UE::HLSL
+{
+#endif
+
+struct FVSMVisibleInstanceCmd
+{
+	uint PackedPageInfo;
+	uint InstanceIdAndFlags;
+	uint IndirectArgIndex;
+};
+
+struct FVSMCullingBatchInfo
+{
+	uint FirstPrimaryView;
+	uint NumPrimaryViews;
+	uint PrimitiveRevealedOffset;
+	uint PrimitiveRevealedNum;
+};
+
+#ifdef __cplusplus
+} // namespace UE::HLSL
+
+using FVSMVisibleInstanceCmd = UE::HLSL::FVSMVisibleInstanceCmd;
+using FVSMCullingBatchInfo = UE::HLSL::FVSMCullingBatchInfo;
+#endif
\ No newline at end of file
