Fix a static analysis error

diff --git a/Engine/Shaders/Shared/SplineMeshShaderParams.h b/Engine/Shaders/Shared/SplineMeshShaderParams.h
--- a/Engine/Shaders/Shared/SplineMeshShaderParams.h
+++ b/Engine/Shaders/Shared/SplineMeshShaderParams.h
@@ -10,7 +10,7 @@
 
 namespace UE::HLSL
 {
-#endif
+#endif // __cplusplus
 
 struct FSplineMeshShaderParams
 {
@@ -36,10 +36,10 @@ struct FSplineMeshShaderParams
 
 #ifdef __cplusplus
 } // namespace UE::HLSL
-
 using FSplineMeshShaderParams = UE::HLSL::FSplineMeshShaderParams;
-#else
+#endif // __cplusplus
 
+#ifndef __cplusplus
 #include "/Engine/Private/Quaternion.ush"
 
 // HLSL unpack method
@@ -76,5 +76,4 @@ FSplineMeshShaderParams UnpackSplineMeshParams(float4 PackedParams[SPLINE_MESH_P
 
 	return Output;
 }
-
-#endif
+#endif // !__cplusplus
