Substrate - shader version bump update.

#prelfight https://horde.devtools.epicgames.com/job/650d50905454efc7fe0e48ac

diff --git a/Engine/Shaders/Shared/SubstrateDefinitions.h b/Engine/Shaders/Shared/SubstrateDefinitions.h
--- a/Engine/Shaders/Shared/SubstrateDefinitions.h
+++ b/Engine/Shaders/Shared/SubstrateDefinitions.h
@@ -7,8 +7,10 @@
 
 #pragma once
 
-// Change this to force recompilation of all Substrate dependent shaders (use https://www.random.org/cgi-bin/randbyte?nbytes=4&format=h)
-#define SUBSTRATE_SHADER_VERSION 0XA577A11B
+#ifndef __cplusplus
+// Change this to force recompilation of all Substrate dependent shaders (for instance https://guidgenerator.com/online-guid-generator.aspx)
+#pragma message("UESHADERMETADATA_VERSION B8E54812-515D-45F4-8A68-AE2FB6553231")
+#endif
 
 // BSDF offsets are packed into 32bits, each entry using SUBSTRATE_BSDF_OFFSET_BIT_COUNT bits
 #define SUBSTRATE_MAX_BSDF_COUNT_FOR_BDSFOFFSET			8u
