Renable triangle strip for hair HW path.

This make the HW path 2x faster.

diff --git a/Engine/Shaders/Shared/HairStrandsDefinitions.h b/Engine/Shaders/Shared/HairStrandsDefinitions.h
--- a/Engine/Shaders/Shared/HairStrandsDefinitions.h
+++ b/Engine/Shaders/Shared/HairStrandsDefinitions.h
@@ -42,7 +42,7 @@
 #define STRANDS_PROCEDURAL_INTERSECTOR_MAX_SPLITS 4
 
 // Use triangle strip for HW raster path
-#define USE_HAIR_TRIANGLE_STRIP 0
+#define USE_HAIR_TRIANGLE_STRIP 1
 
 // Number of vertex per control-point
 #if USE_HAIR_TRIANGLE_STRIP
