Fix HLSL StaticAssert on DXC

The new shader preprocessor does not define the builtin DXC macros, so we should use the UE defined ones instead.

diff --git a/Engine/Shaders/Shared/HLSLStaticAssert.h b/Engine/Shaders/Shared/HLSLStaticAssert.h
--- a/Engine/Shaders/Shared/HLSLStaticAssert.h
+++ b/Engine/Shaders/Shared/HLSLStaticAssert.h
@@ -14,7 +14,7 @@ HLSL_STATIC_ASSERT(true, "Verify that custom static assert works");
 // Included from C++ code, use C++11 keyword
 #define HLSL_STATIC_ASSERT(Expr, Msg)	static_assert(Expr, Msg)
 
-#elif defined(__hlsl_dx_compiler)
+#elif COMPILER_DXC
 
 // DXC only supports C11 style static asserts
 #define HLSL_STATIC_ASSERT(Expr, Msg) _Static_assert(Expr, Msg)
