Fix HLSL_STATIC_ASSERT not working on some platforms

diff --git a/Engine/Shaders/Shared/HLSLStaticAssert.h b/Engine/Shaders/Shared/HLSLStaticAssert.h
--- a/Engine/Shaders/Shared/HLSLStaticAssert.h
+++ b/Engine/Shaders/Shared/HLSLStaticAssert.h
@@ -9,12 +9,12 @@ HLSL_STATIC_ASSERT(true, "Verify that custom static assert works");
 
 #else
 
-#if defined(__cplusplus)
+#if defined(__cplusplus) || COMPILER_PSSL
 
 // Included from C++ code, use C++11 keyword
 #define HLSL_STATIC_ASSERT(Expr, Msg)	static_assert(Expr, Msg)
 
-#elif COMPILER_DXC
+#elif COMPILER_DXC && !COMPILER_VULKAN
 
 // DXC only supports C11 style static asserts
 #define HLSL_STATIC_ASSERT(Expr, Msg) _Static_assert(Expr, Msg)
