Fix incorrect hair flags value

diff --git a/Engine/Shaders/Shared/HairStrandsDefinitions.h b/Engine/Shaders/Shared/HairStrandsDefinitions.h
--- a/Engine/Shaders/Shared/HairStrandsDefinitions.h
+++ b/Engine/Shaders/Shared/HairStrandsDefinitions.h
@@ -101,11 +101,11 @@
 #define HAIR_CARDS_MAX_TEXTURE_COUNT 6
 
 // Hair instance flags 
-#define HAIR_FLAGS_SCATTER_SCENE_LIGHT 1u
-#define HAIR_FLAGS_STABLE_RASTER 2u
-#define HAIR_FLAGS_RAYTRACING_GEOMETRY 4u
-#define HAIR_FLAGS_HOLDOUT 8u
-#define HAIR_FLAGS_CAST_CONTACT_SHADOW 10u
+#define HAIR_FLAGS_SCATTER_SCENE_LIGHT 0x1u
+#define HAIR_FLAGS_STABLE_RASTER 0x2u
+#define HAIR_FLAGS_RAYTRACING_GEOMETRY 0x4u
+#define HAIR_FLAGS_HOLDOUT 0x8u
+#define HAIR_FLAGS_CAST_CONTACT_SHADOW 0x10u
 
 #ifndef __cplusplus //HLSL
 bool HasHairFlags(uint In, uint Flags)
