Removing "platform checking" code from RayTracingBuiltInResources.h that would always be true since we define them to 0 by default. This makes RayTracingBuiltInResources.h match HLSLReservedSpaces.h

diff --git a/Engine/Shaders/Shared/RayTracingBuiltInResources.h b/Engine/Shaders/Shared/RayTracingBuiltInResources.h
--- a/Engine/Shaders/Shared/RayTracingBuiltInResources.h
+++ b/Engine/Shaders/Shared/RayTracingBuiltInResources.h
@@ -14,12 +14,9 @@
 #if defined(__cplusplus)
 	#define INCLUDED_FROM_CPP_CODE  1
 	#define INCLUDED_FROM_HLSL_CODE 0
-#elif defined(SM5_PROFILE) || defined(VULKAN_PROFILE_SM6)
-	// #dxr_todo: we should use a built-in macro to detect if this shader is compiled using DXC (depends on [REDACTED_PROFILE_URL]/DirectXShaderCompiler/issues/1686)
+#else
 	#define INCLUDED_FROM_CPP_CODE  0
 	#define INCLUDED_FROM_HLSL_CODE 1
-#else
-	#error Unknown Compiler
 #endif
 
 #if INCLUDED_FROM_HLSL_CODE
