Substrate - added more data the substrate debug view (frameindex and glint derivative)

#tests PC dx12 tested
#rnx

diff --git a/Engine/Shaders/Private/RayTracing/RayTracingDebug.usf b/Engine/Shaders/Private/RayTracing/RayTracingDebug.usf
--- a/Engine/Shaders/Private/RayTracing/RayTracingDebug.usf
+++ b/Engine/Shaders/Private/RayTracing/RayTracingDebug.usf
@@ -142,7 +142,7 @@ float4 DebugRayTracingMaterial(uint2 PixelCoord, FRayDesc Ray, bool bIsFarField,
 				const uint ClosureIndex = 0;
 				SubstratePrintMaterialProperties(Ctx, PixelCoord, TranslatedWorldPosition, V, ClosureIndex, Payload.SubstrateData);
 
-				SerializeSubstratePixelDebugDataEntry(PixelCoord, Payload.SubstrateData);
+				SerializeSubstratePixelDebugDataEntry(PixelCoord, Payload.SubstrateData, View.FrameNumber);
 
 			}
 			bWriteOuput = false;
diff --git a/Engine/Shaders/Private/Substrate/SubstrateVisualize.usf b/Engine/Shaders/Private/Substrate/SubstrateVisualize.usf
--- a/Engine/Shaders/Private/Substrate/SubstrateVisualize.usf
+++ b/Engine/Shaders/Private/Substrate/SubstrateVisualize.usf
@@ -69,7 +69,7 @@ void MaterialPrintInfoCS(uint3 DispatchThreadId : SV_DispatchThreadID)
 	SubstratePrintMaterialProperties(Context, PixelPos, WorldPosition, V, 7);
 
 	// Serialise out to buffer for print.
-	SerializeSubstratePixelDebugDataEntry(PixelPos);
+	SerializeSubstratePixelDebugDataEntry(PixelPos, PrimaryView.FrameNumber);
 
 	// Print pixel footprint
 	const float3 WorldPositionX = GetWorldPositionFromPixelPos(PixelPos + uint2(1, 0));
diff --git a/Engine/Shaders/Shared/SubstrateVisualizeDefinitions.h b/Engine/Shaders/Shared/SubstrateVisualizeDefinitions.h
--- a/Engine/Shaders/Shared/SubstrateVisualizeDefinitions.h
+++ b/Engine/Shaders/Shared/SubstrateVisualizeDefinitions.h
@@ -64,6 +64,10 @@ struct FSubsterateDebugClosure
 	float FuzzRoughness;
 
 	float GlintValue;
+	float GlintUVDDXx;
+	float GlintUVDDXy;
+	float GlintUVDDYx;
+	float GlintUVDDYy;
 
 	int SpecProfileID;
 	int SpecProfileParameterization;
@@ -94,6 +98,8 @@ struct FSubstratePixelDebugData
 	int MemorySlotC;
 	int MemorySSSData;
 	int MemoryTotal;
+
+	int GPUFrameNumber;
 };
 
 #ifndef __cplusplus
@@ -218,6 +224,10 @@ void SerializeSubstratePixelDebugData(FSubstratePixelDebugDataSerializer& S, FSu
 		S.Serialize(D.Closures[i].FuzzRoughness);
 
 		S.Serialize(D.Closures[i].GlintValue);
+		S.Serialize(D.Closures[i].GlintUVDDXx);
+		S.Serialize(D.Closures[i].GlintUVDDXy);
+		S.Serialize(D.Closures[i].GlintUVDDYx);
+		S.Serialize(D.Closures[i].GlintUVDDYy);
 
 		S.Serialize(D.Closures[i].SpecProfileID);
 		S.Serialize(D.Closures[i].SpecProfileParameterization);
@@ -230,14 +240,15 @@ void SerializeSubstratePixelDebugData(FSubstratePixelDebugDataSerializer& S, FSu
 	S.Serialize(D.MemorySSSData);
 	S.Serialize(D.MemoryTotal);
 
+	S.Serialize(D.GPUFrameNumber);
 }
 
 #ifndef __cplusplus
 
 #if SUBSTRATE_MATERIALCONTAINER_IS_VIEWRESOURCE
-FSubstratePixelDebugData ConvertToSerializableSubstratePixelDebugData(uint2 InCoord)
+FSubstratePixelDebugData ConvertToSerializableSubstratePixelDebugData(uint2 InCoord, uint GPUFrameNumber)
 #else
-FSubstratePixelDebugData ConvertToSerializableSubstratePixelDebugData(uint2 InCoord, FSubstrateRaytracingPayload Payload)
+FSubstratePixelDebugData ConvertToSerializableSubstratePixelDebugData(uint2 InCoord, FSubstrateRaytracingPayload Payload, uint GPUFrameNumber)
 #endif
 {
 	FSubstratePixelDebugData Data = (FSubstratePixelDebugData)0;
@@ -441,6 +452,10 @@ FSubstratePixelDebugData ConvertToSerializableSubstratePixelDebugData(uint2 InCo
 			if (BSDF_GETHASGLINT(BSDF))
 			{
 				Data.Closures[i].GlintValue = SLAB_GLINT_VALUE(BSDF);
+				Data.Closures[i].GlintUVDDXx = SLAB_GLINT_UVDDX(BSDF).x;
+				Data.Closures[i].GlintUVDDXy = SLAB_GLINT_UVDDX(BSDF).y;
+				Data.Closures[i].GlintUVDDYx = SLAB_GLINT_UVDDY(BSDF).x;
+				Data.Closures[i].GlintUVDDYy = SLAB_GLINT_UVDDY(BSDF).y;
 			}
 			else
 			{
@@ -567,20 +582,22 @@ FSubstratePixelDebugData ConvertToSerializableSubstratePixelDebugData(uint2 InCo
 		Data.MemoryTotal = TotalSize;
 	}
 
+	Data.GPUFrameNumber = asint(GPUFrameNumber);
+
 	return Data;
 }
 
 #if SUBSTRATE_MATERIALCONTAINER_IS_VIEWRESOURCE
-void SerializeSubstratePixelDebugDataEntry(uint2 InCoord)
+void SerializeSubstratePixelDebugDataEntry(uint2 InCoord, uint GPUFrameNumber)
 #else
-void SerializeSubstratePixelDebugDataEntry(uint2 InCoord, FSubstrateRaytracingPayload Payload)
+void SerializeSubstratePixelDebugDataEntry(uint2 InCoord, FSubstrateRaytracingPayload Payload, uint GPUFrameNumber)
 #endif
 {
 	// Convert Substrate GBuffer into the CPU/GPU common data to serialize.
 #if SUBSTRATE_MATERIALCONTAINER_IS_VIEWRESOURCE
-	FSubstratePixelDebugData Data = ConvertToSerializableSubstratePixelDebugData(InCoord);
+	FSubstratePixelDebugData Data = ConvertToSerializableSubstratePixelDebugData(InCoord, GPUFrameNumber);
 #else
-	FSubstratePixelDebugData Data = ConvertToSerializableSubstratePixelDebugData(InCoord, Payload);
+	FSubstratePixelDebugData Data = ConvertToSerializableSubstratePixelDebugData(InCoord, Payload, GPUFrameNumber);
 #endif
 	
 	// And serialise it.
diff --git a/Engine/Source/Runtime/Renderer/Private/Substrate/SubstrateVisualize.cpp b/Engine/Source/Runtime/Renderer/Private/Substrate/SubstrateVisualize.cpp
--- a/Engine/Source/Runtime/Renderer/Private/Substrate/SubstrateVisualize.cpp
+++ b/Engine/Source/Runtime/Renderer/Private/Substrate/SubstrateVisualize.cpp
@@ -410,6 +410,12 @@ void AddProcessAndPrintSubstrateMaterialPropertiesPasses(FRDGBuilder& GraphBuild
 					TextItemsBatch.Push(FCanvasTextItem(FVector2D(X, Y), FText::FromString(*String), GetStatsFont(), Color));
 					ExpendBatchBounds(X, Y);
 				};
+				auto PrintUI = [&](float X, float Y, uint32 Value, FLinearColor Color = FLinearColor::White)
+				{
+					FString String = FString::Printf(TEXT("%u"), Value);
+					TextItemsBatch.Push(FCanvasTextItem(FVector2D(X, Y), FText::FromString(*String), GetStatsFont(), Color));
+					ExpendBatchBounds(X, Y);
+				};
 				auto PrintFSmall = [&](float X, float Y, float Value, FLinearColor Color = FLinearColor::White)
 				{
 					FString String = FString::Printf(TEXT("%1.3f"), Value);
@@ -451,6 +457,11 @@ void AddProcessAndPrintSubstrateMaterialPropertiesPasses(FRDGBuilder& GraphBuild
 					PrintFAdapt(X + 50.0f, Y, G, CompG);
 					PrintFAdapt(X + 100.0f, Y, B, CompB);
 				};
+				auto PrintF2DAdapt = [&](float X, float Y, float R, float G)
+				{
+					PrintFAdapt(X + 0.0f, Y, R, CompR);
+					PrintFAdapt(X + 50.0f, Y, G, CompG);
+				};
 				auto PrintFBool = [&](float X, float Y, float Value)
 				{
 					FString String = Value > 0.0f ? TEXT("Yes") : TEXT("No");
@@ -732,9 +743,18 @@ void AddProcessAndPrintSubstrateMaterialPropertiesPasses(FRDGBuilder& GraphBuild
 
 						if (Data.Closures[i].GlintValue < 1.0f)
 						{
-							PrintS(DrawPosX + 0.0f, DrawPosY, TEXT("Glint Density"), Orange);
+							PrintS(DrawPosX + 0.0f, DrawPosY, TEXT("Glints"), Orange);
+							NewLine();
+
+							PrintS(DrawPosX + 0.0f, DrawPosY, TEXT(" - Density"), Orange);
 							PrintFSmall(DrawPosX + DataOffset, DrawPosY, Data.Closures[i].GlintValue);
 							NewLine();
+							PrintS(DrawPosX + 0.0f, DrawPosY, TEXT(" - UV DDX"), Orange);
+							PrintF2DAdapt(DrawPosX + DataOffset, DrawPosY, Data.Closures[i].GlintUVDDXx, Data.Closures[i].GlintUVDDXy);
+							NewLine();
+							PrintS(DrawPosX + 0.0f, DrawPosY, TEXT(" - UV DDY"), Orange);
+							PrintF2DAdapt(DrawPosX + DataOffset, DrawPosY, Data.Closures[i].GlintUVDDYx, Data.Closures[i].GlintUVDDYy);
+							NewLine();
 						}
 
 						if (Data.Closures[i].SpecProfileID >= 0)
@@ -893,9 +913,19 @@ void AddProcessAndPrintSubstrateMaterialPropertiesPasses(FRDGBuilder& GraphBuild
 					PrintS(DrawPosX + 0.0f, DrawPosY, TEXT(" - Total"), Orange);
 					PrintI(DrawPosX + DataOffset, DrawPosY, Data.MemoryTotal);
 					NewLine();
+					NewLine();
 
 					DispatchBatch();
 				}
+
+				BeginBatch();
+				PrintS(DrawPosX + 0.0f, DrawPosY, TEXT("GPUFrameNumber"), FLinearColor::White);
+				PrintUI(DrawPosX + 150, DrawPosY, *(uint32*)(&Data.GPUFrameNumber));
+				NewLine();
+				PrintS(DrawPosX + 0.0f, DrawPosY, TEXT("CPUFrameNumber"), FLinearColor::White);
+				PrintUI(DrawPosX + 150, DrawPosY, View.CachedViewUniformShaderParameters->FrameNumber);
+				NewLine();
+				DispatchBatch();
 			});
 		////////////////////////////////////////////////////////////////////////////////////////////////////
 
