Substrate - fixed shader compiilation errors.

#tests Tested pc dx12 with and without substrate, with RT
#rnx

diff --git a/Engine/Shaders/Private/Substrate/Glint/GlintThirdParty.ush b/Engine/Shaders/Private/Substrate/Glint/GlintThirdParty.ush
--- a/Engine/Shaders/Private/Substrate/Glint/GlintThirdParty.ush
+++ b/Engine/Shaders/Private/Substrate/Glint/GlintThirdParty.ush
@@ -813,21 +813,21 @@ float2 Sample_P22_LOD(
 		ValidCellCount = min(ValidCellCount, MAX_CELL_COUNT);
 
 		// 1.2 Normalized CDF
-		uint it;
+		uint CellIt;
 		LOOP
-		for (it = 0; it < ValidCellCount; ++it)
+		for (CellIt = 0; CellIt < ValidCellCount; ++CellIt)
 		{
-			Cdfs[it] /= max(Acc, 0.0001f);
+			Cdfs[CellIt] /= max(Acc, 0.0001f);
 		}
 		Cdfs[ValidCellCount] = 1;
 
 		// 1.3 Sample CDF
 		LOOP
-		for (it = 0; it < ValidCellCount; ++it)
+		for (CellIt = 0; CellIt < ValidCellCount; ++CellIt)
 		{
-			if (Cdfs[it] <= u0 && u0 <= Cdfs[it + 1])
+			if (Cdfs[CellIt] <= u0 && u0 <= Cdfs[CellIt + 1])
 			{
-				selected_st = Indices[it];
+				selected_st = Indices[CellIt];
 				break;
 			}
 		}
diff --git a/Engine/Shaders/Shared/SubstrateVisualizeDefinitions.h b/Engine/Shaders/Shared/SubstrateVisualizeDefinitions.h
--- a/Engine/Shaders/Shared/SubstrateVisualizeDefinitions.h
+++ b/Engine/Shaders/Shared/SubstrateVisualizeDefinitions.h
@@ -7,10 +7,11 @@
 
 #pragma once
 
+#if defined(__cplusplus) || SUBSTRATE_ENABLED
 
 #ifndef __cplusplus
 // Change this to force recompilation of all Substrate dependent shaders (for instance https://guidgenerator.com/online-guid-generator.aspx)
-#pragma message("UESHADERMETADATA_VERSION 8FFAABAF-705E-481F-8234-D4D42C7A4001")
+#pragma message("UESHADERMETADATA_VERSION E9DF359B-BC55-46BC-B4DF-A5E94467C029")
 #include "/Engine/Shared/SubstrateDefinitions.h"
 #else
 #include "SubstrateDefinitions.h"
@@ -297,7 +298,9 @@ FSubstratePixelDebugData ConvertToSerializableSubstratePixelDebugData(uint2 InCo
 	Data.DoesCastContactShadow = Header.DoesCastContactShadow() ? 1 : 0;
 	Data.HasDynamicIndirectShadowCasterRepresentation = Header.HasDynamicIndirectShadowCasterRepresentation() ? 1 : 0;
 	Data.HasSubsurface = Header.HasSubsurface() ? 1 : 0;
+#if SUBSTRATE_DEFERRED_SHADING
 	Data.LocalBasesCount = Header.GetLocalBasesCount();
+#endif
 
 	LOOP
 	for(uint i = 0; i < Header.ClosureCount; ++i)
@@ -312,9 +315,10 @@ FSubstratePixelDebugData ConvertToSerializableSubstratePixelDebugData(uint2 InCo
 		Data.Closures[i].Address = SubstrateAddressing.CurrentIndex;
 		Data.Closures[i].Type = BSDF_GETTYPE(BSDF);
 
+#if SUBSTRATE_DEFERRED_SHADING
 		Data.Closures[i].NormalID = BSDF_GETSHAREDLOCALBASISID(BSDF);
 		Data.Closures[i].BasisType = SubstrateGetSharedLocalBasisType(HEADER_GETSHAREDLOCALBASESTYPE(Header.PackedHeader), BSDF_GETSHAREDLOCALBASISID(BSDF));
-
+#endif
 
 		Data.Closures[i].bIsTopLayer = BSDF_GETISTOPLAYER(BSDF);
 		Data.Closures[i].bHasWeightL = BSDF_GETHASTRANSABOVE(BSDF);
@@ -650,3 +654,4 @@ void ConvertToSerializableSubstratePixelDebugData()
 
 #endif
 
+#endif // defined(__cplusplus) || SUBSTRATE_ENABLED
