[GPU Scene] Fix incorrect packing in FLightSceneData.
- DiffuseScale was causing a mismatch between C++ and Shader layout.

diff --git a/Engine/Shaders/Shared/LightSceneData.h b/Engine/Shaders/Shared/LightSceneData.h
--- a/Engine/Shaders/Shared/LightSceneData.h
+++ b/Engine/Shaders/Shared/LightSceneData.h
@@ -33,11 +33,8 @@ struct FLightSceneData
 	// Direction of the light if applies.
 	float3 Direction;
 
-	// Factor to applies on the specular.
-	float SpecularScale;
-
-	// Factor to applies on the diffuse.
-	float DiffuseScale;
+	// Extra fields
+	uint LightTypeAndShadowMapChannelMaskPacked;
 
 	// One tangent of the light if applies.
 	// Note: BiTangent is on purpose not stored for memory optimisation purposes.
@@ -73,10 +70,13 @@ struct FLightSceneData
 	// could probably go down to 8 bits (with some logic in GIESTextureManager to warn about overflow)
 	float IESAtlasIndex;
 
-	// Extra fields
-	uint LightTypeAndShadowMapChannelMaskPacked;
+	// Factor to applies on the specular.
+	float SpecularScale;
+
+	// Factor to applies on the diffuse.
+	float DiffuseScale;
 
-	float2 Padding;
+	float Padding;
 };
 
 #ifdef __cplusplus
