[Shaders] Don't use sizeof() in shaders that target PCD3D_SM5 since it's a reserved word for FXC.

This keyword was introduced in SceneData.ush with CL 37830107.

#rnx

diff --git a/Engine/Shaders/Private/SceneData.ush b/Engine/Shaders/Private/SceneData.ush
--- a/Engine/Shaders/Private/SceneData.ush
+++ b/Engine/Shaders/Private/SceneData.ush
@@ -1403,11 +1403,11 @@ ByteAddressBuffer GPUSceneLights;
 FLightSceneData GetLightSceneData(uint LightId)
 {
 #if 0
-	return GPUSceneLightsBuffer.Load<FLightSceneData>(LightId * sizeof(FLightSceneData));
+	return GPUSceneLightsBuffer.Load<FLightSceneData>(LightId * SIZEOF_LIGHT_SCENE_DATA);
 #else
 	FLightSceneData Output;
 
-	uint Offset = LightId * (uint)sizeof(FLightSceneData);
+	uint Offset = LightId * (uint)SIZEOF_LIGHT_SCENE_DATA;
 
 	LoadAndIncrementOffset(Output.WorldPosition.High, GPUSceneLightsBuffer, Offset);
 	LoadAndIncrementOffset(Output.WorldPosition.Low, GPUSceneLightsBuffer, Offset);
diff --git a/Engine/Shaders/Shared/LightSceneData.h b/Engine/Shaders/Shared/LightSceneData.h
--- a/Engine/Shaders/Shared/LightSceneData.h
+++ b/Engine/Shaders/Shared/LightSceneData.h
@@ -61,6 +61,35 @@ struct FLightSceneData
 	float DiffuseScale;
 };
 
+// Alternative for sizeof(FLightSceneData). FXC reserves keyword 'sizeof' so we cannot use it for shader permutations that target PCD3D_SM5.
+//  struct struct.FLightSceneData
+//  {
+//
+//      struct struct.FDFVector3
+//      {
+//
+//          float3 a;                                 ; Offset:    0
+//          float3 b;                                 ; Offset:   12
+//      
+//      } WorldPosition;                              ; Offset:    0
+//
+//      float InvRadius;                              ; Offset:   24
+//      float FalloffExponent;                        ; Offset:   28
+//      float3 Direction;                             ; Offset:   32
+//      float InverseExposureBlend;                   ; Offset:   44
+//      float3 Tangent;                               ; Offset:   48
+//      float SourceRadius;                           ; Offset:   60
+//      float2 SpotAngles;                            ; Offset:   64
+//      float SoftSourceRadius;                       ; Offset:   72
+//      float SourceLength;                           ; Offset:   76
+//      float RectLightBarnCosAngle;                  ; Offset:   80
+//      float RectLightBarnLength;                    ; Offset:   84
+//      float SpecularScale;                          ; Offset:   88
+//      float DiffuseScale;                           ; Offset:   92
+//  
+//  } $Element;                                       ; Offset:    0 Size:    96
+#define SIZEOF_LIGHT_SCENE_DATA 96
+
 #ifdef __cplusplus
 } // namespace
 using FLightSceneData = UE::HLSL::FLightSceneData;
diff --git a/Engine/Source/Runtime/Renderer/Private/GPUScene.cpp b/Engine/Source/Runtime/Renderer/Private/GPUScene.cpp
--- a/Engine/Source/Runtime/Renderer/Private/GPUScene.cpp
+++ b/Engine/Source/Runtime/Renderer/Private/GPUScene.cpp
@@ -36,6 +36,12 @@
 #define FORCEINLINE_GPUSCENE FORCEINLINE
 //#define FORCEINLINE_GPUSCENE
 
+static_assert(
+	sizeof(FLightSceneData) == SIZEOF_LIGHT_SCENE_DATA,
+	"FLightSceneData has changed without update to SIZEOF_LIGHT_SCENE_DATA. "
+	"This macro is required because sizeof() is not available in FXC and therefore not available for PCD3D_SM5 targets."
+);
+
 // Defaults to being disabled, enable using the command line argument: -CsvCategory GPUScene
 CSV_DEFINE_CATEGORY(GPUScene, false);
 
