[HWRT] Check whether ray tracing shaders or inline RT are supported when setting up ray tracing visualization modes.
- logic wasn't taking whether project/platform doesn't use "full" ray tracing (ie: RTPSOs).

diff --git a/Engine/Shaders/Shared/RaytracingDebugDefinitions.h b/Engine/Shaders/Shared/RaytracingDebugDefinitions.h
--- a/Engine/Shaders/Shared/RaytracingDebugDefinitions.h
+++ b/Engine/Shaders/Shared/RaytracingDebugDefinitions.h
@@ -59,3 +59,50 @@
 #define RAY_TRACING_DEBUG_PICKER_DOMAIN_SEGMENT					2
 #define RAY_TRACING_DEBUG_PICKER_DOMAIN_FLAGS                   3
 #define RAY_TRACING_DEBUG_PICKER_DOMAIN_MASK                    4
+
+#if defined(__cplusplus)
+	#define UINT_TYPE unsigned int
+	#define INT_TYPE int
+	#define INLINE_ATTR inline
+#else
+	#define UINT_TYPE uint
+	#define INT_TYPE int
+	#define INLINE_ATTR 
+#endif
+
+INLINE_ATTR bool RequiresRayTracingDebugCHS(UINT_TYPE DebugVisualizationMode)
+{
+	return DebugVisualizationMode == RAY_TRACING_DEBUG_VIZ_INSTANCES
+		|| DebugVisualizationMode == RAY_TRACING_DEBUG_VIZ_TRIANGLES
+		|| DebugVisualizationMode == RAY_TRACING_DEBUG_VIZ_DYNAMIC_INSTANCES
+		|| DebugVisualizationMode == RAY_TRACING_DEBUG_VIZ_PROXY_TYPE
+		|| DebugVisualizationMode == RAY_TRACING_DEBUG_VIZ_PICKER
+		|| DebugVisualizationMode == RAY_TRACING_DEBUG_VIZ_INSTANCE_OVERLAP
+		|| DebugVisualizationMode == RAY_TRACING_DEBUG_VIZ_TRIANGLE_HITCOUNT
+		|| DebugVisualizationMode == RAY_TRACING_DEBUG_VIZ_HITCOUNT_PER_INSTANCE;
+}
+
+INLINE_ATTR bool IsRayTracingDebugTraversalMode(UINT_TYPE DebugVisualizationMode)
+{
+	return DebugVisualizationMode == RAY_TRACING_DEBUG_VIZ_TRAVERSAL_PRIMARY_NODE
+		|| DebugVisualizationMode == RAY_TRACING_DEBUG_VIZ_TRAVERSAL_PRIMARY_CLUSTER
+		|| DebugVisualizationMode == RAY_TRACING_DEBUG_VIZ_TRAVERSAL_PRIMARY_TRIANGLE
+		|| DebugVisualizationMode == RAY_TRACING_DEBUG_VIZ_TRAVERSAL_PRIMARY_ALL
+		|| DebugVisualizationMode == RAY_TRACING_DEBUG_VIZ_TRAVERSAL_PRIMARY_STATISTICS
+		|| DebugVisualizationMode == RAY_TRACING_DEBUG_VIZ_TRAVERSAL_SECONDARY_NODE
+		|| DebugVisualizationMode == RAY_TRACING_DEBUG_VIZ_TRAVERSAL_SECONDARY_CLUSTER
+		|| DebugVisualizationMode == RAY_TRACING_DEBUG_VIZ_TRAVERSAL_SECONDARY_TRIANGLE
+		|| DebugVisualizationMode == RAY_TRACING_DEBUG_VIZ_TRAVERSAL_SECONDARY_ALL
+		|| DebugVisualizationMode == RAY_TRACING_DEBUG_VIZ_TRAVERSAL_SECONDARY_STATISTICS;
+}
+
+INLINE_ATTR bool RayTracingDebugModeSupportsInline(UINT_TYPE DebugVisualizationMode)
+{
+	const bool bIsBarycentricsMode = DebugVisualizationMode == RAY_TRACING_DEBUG_VIZ_BARYCENTRICS;
+	const bool bUsesDebugCHSAndNotHitCount = RequiresRayTracingDebugCHS(DebugVisualizationMode)
+		&& (DebugVisualizationMode != RAY_TRACING_DEBUG_VIZ_TRIANGLE_HITCOUNT)
+		&& (DebugVisualizationMode != RAY_TRACING_DEBUG_VIZ_HITCOUNT_PER_INSTANCE);
+	const bool bIsTraversalMode = IsRayTracingDebugTraversalMode(DebugVisualizationMode);
+
+	return bIsBarycentricsMode || bUsesDebugCHSAndNotHitCount || bIsTraversalMode;
+}
diff --git a/Engine/Source/Runtime/Engine/Private/RayTracingVisualizationData.cpp b/Engine/Source/Runtime/Engine/Private/RayTracingVisualizationData.cpp
--- a/Engine/Source/Runtime/Engine/Private/RayTracingVisualizationData.cpp
+++ b/Engine/Source/Runtime/Engine/Private/RayTracingVisualizationData.cpp
@@ -21,12 +21,9 @@ static FAutoConsoleVariableDeprecated CVarRayTracingDebugMode_Deprecated(TEXT("r
 
 FRayTracingVisualizationData::FRayTracingVisualizationData()
 {
-	const bool bRayTracingShadersSupported = FDataDrivenShaderPlatformInfo::GetSupportsRayTracingShaders(GMaxRHIShaderPlatform);
-
 	// always supported (as long as either inline RT or RT shaders work)
 	AddVisualizationMode(TEXT("Barycentrics"), LOCTEXT("Barycentrics", "Barycentrics"), FModeType::Standard, RAY_TRACING_DEBUG_VIZ_BARYCENTRICS, false);
 
-	if (bRayTracingShadersSupported)
 	{
 		// runs basic lighting calculations on hits
 		AddVisualizationMode(TEXT("PrimaryRays"), LOCTEXT("PrimaryRays", "Primary Rays"), FModeType::Standard, RAY_TRACING_DEBUG_VIZ_PRIMARY_RAYS, true);
@@ -104,6 +101,11 @@ void FRayTracingVisualizationData::AddVisualizationMode(
 	bool bHiddenInEditor
 )
 {
+	if (!GRHIGlobals.RayTracing.SupportsShaders && !(GRHIGlobals.RayTracing.SupportsInlineRayTracing && RayTracingDebugModeSupportsInline(ModeID)))
+	{
+		return;
+	}
+
 	const FName ModeName = FName(ModeString);
 
 	FModeRecord& Record	= ModeMap.Emplace(ModeName);
diff --git a/Engine/Source/Runtime/Renderer/Private/RayTracing/RayTracingDebug.cpp b/Engine/Source/Runtime/Renderer/Private/RayTracing/RayTracingDebug.cpp
--- a/Engine/Source/Runtime/Renderer/Private/RayTracing/RayTracingDebug.cpp
+++ b/Engine/Source/Runtime/Renderer/Private/RayTracing/RayTracingDebug.cpp
@@ -917,39 +917,6 @@ static void BindRayTracingDebugCHSMaterialBindings(
 		bCopyDataToInlineStorage);
 }
 
-static bool RequiresRayTracingDebugCHS(uint32 DebugVisualizationMode)
-{
-	return DebugVisualizationMode == RAY_TRACING_DEBUG_VIZ_INSTANCES
-		|| DebugVisualizationMode == RAY_TRACING_DEBUG_VIZ_TRIANGLES
-		|| DebugVisualizationMode == RAY_TRACING_DEBUG_VIZ_DYNAMIC_INSTANCES
-		|| DebugVisualizationMode == RAY_TRACING_DEBUG_VIZ_PROXY_TYPE
-		|| DebugVisualizationMode == RAY_TRACING_DEBUG_VIZ_PICKER
-		|| DebugVisualizationMode == RAY_TRACING_DEBUG_VIZ_INSTANCE_OVERLAP
-		|| DebugVisualizationMode == RAY_TRACING_DEBUG_VIZ_TRIANGLE_HITCOUNT
-		|| DebugVisualizationMode == RAY_TRACING_DEBUG_VIZ_HITCOUNT_PER_INSTANCE;
-}
-
-static bool RayTracingDebugModeSupportsInline(uint32 DebugVisualizationMode)
-{
-	return RequiresRayTracingDebugCHS(DebugVisualizationMode)
-		&& (DebugVisualizationMode != RAY_TRACING_DEBUG_VIZ_TRIANGLE_HITCOUNT)
-		&& (DebugVisualizationMode != RAY_TRACING_DEBUG_VIZ_HITCOUNT_PER_INSTANCE);
-}
-
-static bool IsRayTracingDebugTraversalMode(uint32 DebugVisualizationMode)
-{
-	return DebugVisualizationMode == RAY_TRACING_DEBUG_VIZ_TRAVERSAL_PRIMARY_NODE 
-		|| DebugVisualizationMode == RAY_TRACING_DEBUG_VIZ_TRAVERSAL_PRIMARY_CLUSTER
-		|| DebugVisualizationMode == RAY_TRACING_DEBUG_VIZ_TRAVERSAL_PRIMARY_TRIANGLE
-		|| DebugVisualizationMode == RAY_TRACING_DEBUG_VIZ_TRAVERSAL_PRIMARY_ALL
-		|| DebugVisualizationMode == RAY_TRACING_DEBUG_VIZ_TRAVERSAL_PRIMARY_STATISTICS
-		|| DebugVisualizationMode == RAY_TRACING_DEBUG_VIZ_TRAVERSAL_SECONDARY_NODE
-		|| DebugVisualizationMode == RAY_TRACING_DEBUG_VIZ_TRAVERSAL_SECONDARY_CLUSTER
-		|| DebugVisualizationMode == RAY_TRACING_DEBUG_VIZ_TRAVERSAL_SECONDARY_TRIANGLE
-		|| DebugVisualizationMode == RAY_TRACING_DEBUG_VIZ_TRAVERSAL_SECONDARY_ALL
-		|| DebugVisualizationMode == RAY_TRACING_DEBUG_VIZ_TRAVERSAL_SECONDARY_STATISTICS;
-}
-
 static bool IsRayTracingPickingEnabled(uint32 DebugVisualizationMode)
 {
 	return DebugVisualizationMode == RAY_TRACING_DEBUG_VIZ_PICKER;
