Fix compilation issue allowing to renable Substrate visualization on PS5

Tested on SubstrateTest with different Closure count for PS5

diff --git a/Engine/Shaders/Private/Substrate/SubstrateVisualize.usf b/Engine/Shaders/Private/Substrate/SubstrateVisualize.usf
--- a/Engine/Shaders/Private/Substrate/SubstrateVisualize.usf
+++ b/Engine/Shaders/Private/Substrate/SubstrateVisualize.usf
@@ -78,10 +78,7 @@ void MaterialPrintInfoCS(uint3 DispatchThreadId : SV_DispatchThreadID)
 	SubstratePrintMaterialProperties(Context, PixelPos, WorldPosition, V, 7);
 
 	// Serialise out to buffer for print.
-	// Temporarely disable data serialization on PSSL due to internal shader compilation 
-	#if !COMPILER_PSSL
 	SerializeSubstratePixelDebugDataEntry(PixelPos, PrimaryView.FrameNumber);
-	#endif
 
 #endif
 
diff --git a/Engine/Shaders/Shared/SubstrateVisualizeDefinitions.h b/Engine/Shaders/Shared/SubstrateVisualizeDefinitions.h
--- a/Engine/Shaders/Shared/SubstrateVisualizeDefinitions.h
+++ b/Engine/Shaders/Shared/SubstrateVisualizeDefinitions.h
@@ -309,8 +309,14 @@ FSubstratePixelDebugData ConvertToSerializableSubstratePixelDebugData(uint2 InCo
 	Data.LocalBasesCount = Header.GetLocalBasesCount();
 #endif
 
+	// PSSL - Temporarely limit data serialization on PSSL due to internal shader compilation 
+	#if COMPILER_PSSL
+	SUBSTRATE_UNROLL_N(SUBSTRATE_MAX_CLOSURE_COUNT)
+	for(uint i = 0; i < Header.ClosureCount; ++i)
+	#else
 	LOOP
 	for(uint i = 0; i < Header.ClosureCount; ++i)
+	#endif
 	{
 		// Unpack BSDF data
 	#if SUBSTRATE_MATERIALCONTAINER_IS_VIEWRESOURCE
