Fix minor shader warnings in Substrate materials

#rnx

diff --git a/Engine/Shaders/Shared/SubstrateVisualizeDefinitions.h b/Engine/Shaders/Shared/SubstrateVisualizeDefinitions.h
--- a/Engine/Shaders/Shared/SubstrateVisualizeDefinitions.h
+++ b/Engine/Shaders/Shared/SubstrateVisualizeDefinitions.h
@@ -508,7 +508,7 @@ FSubstratePixelDebugData ConvertToSerializableSubstratePixelDebugData(uint2 InCo
 			Data.Closures[i].DiffuseR	= EYE_DIFFUSEALBEDO(BSDF).r;
 			Data.Closures[i].DiffuseG	= EYE_DIFFUSEALBEDO(BSDF).g;
 			Data.Closures[i].DiffuseB	= EYE_DIFFUSEALBEDO(BSDF).b;
-			Data.Closures[i].F0R		= EYE_F0(BSDF);
+			Data.Closures[i].F0R		= EYE_F0(BSDF).r;
 			Data.Closures[i].Roughness	= EYE_ROUGHNESS(BSDF);
 			Data.Closures[i].F90R		= EYE_IRISMASK(BSDF);
 			Data.Closures[i].F90G		= EYE_IRISDISTANCE(BSDF);
