Substrate - Added HasBottom normal to material propertie viz.

#rnx
#tests PC editor dx12

diff --git a/Engine/Shaders/Shared/SubstrateVisualizeDefinitions.h b/Engine/Shaders/Shared/SubstrateVisualizeDefinitions.h
--- a/Engine/Shaders/Shared/SubstrateVisualizeDefinitions.h
+++ b/Engine/Shaders/Shared/SubstrateVisualizeDefinitions.h
@@ -51,6 +51,7 @@ struct FSubsterateDebugClosure
 	float HazeRoughness;
 	float HazeWeight;
 	int HazeSimpleClearCoatMode;
+	int HasBottomNormal;
 
 	float  SSSOpacity;
 	float SSSMFPR, SSSMFPG, SSSMFPB;
@@ -212,6 +213,7 @@ void SerializeSubstratePixelDebugData(FSubstrateDebugDataSerializer& S, FSubstra
 		S.Serialize(D.Closures[i].HazeRoughness);
 		S.Serialize(D.Closures[i].HazeWeight);
 		S.Serialize(D.Closures[i].HazeSimpleClearCoatMode);
+		S.Serialize(D.Closures[i].HasBottomNormal);
 
 		S.Serialize(D.Closures[i].SSSOpacity);
 		S.Serialize(D.Closures[i].SSSMFPR);
@@ -387,6 +389,7 @@ FSubstratePixelDebugData ConvertToSerializableSubstratePixelDebugData(uint2 InCo
 				Data.Closures[i].HazeRoughness = Haziness.Roughness;
 				Data.Closures[i].HazeWeight = Haziness.Weight;
 				Data.Closures[i].HazeSimpleClearCoatMode = Haziness.bSimpleClearCoat ? 1u : 0u;
+				Data.Closures[i].HasBottomNormal = Haziness.HasBottomNormal() ? 1u : 0u;
 			}
 
 			if (Data.Closures[i].SSSType != SSS_TYPE_NONE || Data.Closures[i].bIsThin)
diff --git a/Engine/Source/Runtime/Renderer/Private/Substrate/SubstrateVisualize.cpp b/Engine/Source/Runtime/Renderer/Private/Substrate/SubstrateVisualize.cpp
--- a/Engine/Source/Runtime/Renderer/Private/Substrate/SubstrateVisualize.cpp
+++ b/Engine/Source/Runtime/Renderer/Private/Substrate/SubstrateVisualize.cpp
@@ -609,6 +609,10 @@ void AddProcessAndPrintSubstrateMaterialPropertiesPasses(FRDGBuilder& GraphBuild
 							PrintS(DrawPosX + 0.0f, DrawPosY, TEXT(" - LegacyClearCoat"), Orange);
 							PrintFBool(DrawPosX + DataOffset, DrawPosY, Data.Closures[i].HazeSimpleClearCoatMode);
 							NewLine();
+
+							PrintS(DrawPosX + 0.0f, DrawPosY, TEXT(" - HasBottomNormal"), Orange);
+							PrintFBool(DrawPosX + DataOffset, DrawPosY, Data.Closures[i].HasBottomNormal);
+							NewLine();
 						}
 
 						if (Closure.SSSType != SSS_TYPE_NONE || Closure.bIsThin)
