Material Cache, HLSL cross-compiler workaround for type aliasing

diff --git a/Engine/Shaders/Shared/MaterialCacheDefinitions.h b/Engine/Shaders/Shared/MaterialCacheDefinitions.h
--- a/Engine/Shaders/Shared/MaterialCacheDefinitions.h
+++ b/Engine/Shaders/Shared/MaterialCacheDefinitions.h
@@ -44,7 +44,8 @@ enum EMaterialCacheFlag
 	MatCache_DefaultBottomLayer = 1u << 0u,
 };
 #else // __cplusplus
-typedef uint EMaterialCacheFlag;
+// HLSL Cross-Compiler Workaround
+#define EMaterialCacheFlag uint
 static uint MatCache_None = 0u;
 static uint MatCache_DefaultBottomLayer = 1u << 0u;
 #endif // __cplusplus
