* Fix BSDF state visualization.
* Fix MFP not being clamped when using Format=0
* Fix incorrect reported memory

diff --git a/Engine/Shaders/Private/DeferredShadingCommon.ush b/Engine/Shaders/Private/DeferredShadingCommon.ush
--- a/Engine/Shaders/Private/DeferredShadingCommon.ush
+++ b/Engine/Shaders/Private/DeferredShadingCommon.ush
@@ -418,6 +418,24 @@ bool IsFirstPerson(FGBufferData GBufferData)
 #endif
 }
 
+bool HasPrecShadowMask(FGBufferData GBufferData)
+{
+#if ALLOW_STATIC_LIGHTING
+	return !(GBufferData.SelectiveOutputMask & SKIP_PRECSHADOW_MASK);
+#else
+	return false;
+#endif
+}
+
+bool HasZeroPrecShadowMask(FGBufferData GBufferData)
+{
+#if ALLOW_STATIC_LIGHTING
+	return GBufferData.SelectiveOutputMask & ZERO_PRECSHADOW_MASK;
+#else
+	return false;
+#endif
+}
+
 // High frequency Checkerboard pattern
 // @param PixelPos relative to left top of the rendertarget (not viewport)
 // @return true/false, todo: profile if float 0/1 would be better (need to make sure it's 100% the same)
@@ -1003,7 +1021,7 @@ FGBufferData DecodeGBufferData(
 
 	// FirstPerson uses a bit in SelectiveOutputMask that is aliased with ZERO_PRECSHADOW_MASK when !ALLOW_STATIC_LIGHTING, so we explicitly skip this logic here.
 #if ALLOW_STATIC_LIGHTING
-	GBuffer.PrecomputedShadowFactors = !(GBuffer.SelectiveOutputMask & SKIP_PRECSHADOW_MASK) ? InGBufferE :  ((GBuffer.SelectiveOutputMask & ZERO_PRECSHADOW_MASK) ? 0 :  1);
+	GBuffer.PrecomputedShadowFactors = HasPrecShadowMask(GBuffer) ? InGBufferE : (HasZeroPrecShadowMask(GBuffer) ? 0 : 1);
 #else
 	GBuffer.PrecomputedShadowFactors = half(1.0f);
 #endif
diff --git a/Engine/Shaders/Private/Substrate/Substrate.ush b/Engine/Shaders/Private/Substrate/Substrate.ush
--- a/Engine/Shaders/Private/Substrate/Substrate.ush
+++ b/Engine/Shaders/Private/Substrate/Substrate.ush
@@ -3317,6 +3317,11 @@ float SanitizeF0(float F0)
 	#endif
 }
 
+float3 SanitizeSSSMFP(float3 InMFP)
+{
+	return clamp(InMFP, 0.0f, Max111110BitsFloat3);
+}
+
 void FSubstrateBSDF::SubstrateSanitizeBSDF()
 {
 	switch (BSDF_GETTYPE(this))
@@ -3352,7 +3357,7 @@ void FSubstrateBSDF::SubstrateSanitizeBSDF()
 			}
 			else if (BSDF_GETSSSTYPE(this) == SSS_TYPE_DIFFUSION)
 			{
-				SLAB_SSSMFP(this) = clamp(SLAB_SSSMFP(this), 0.0f, Max111110BitsFloat3);
+				SLAB_SSSMFP(this) = SanitizeSSSMFP(SLAB_SSSMFP(this));
 			}
 
 			SLAB_FUZZ_AMOUNT(this) = saturate(SLAB_FUZZ_AMOUNT(this));
diff --git a/Engine/Shaders/Private/Substrate/SubstrateRead.ush b/Engine/Shaders/Private/Substrate/SubstrateRead.ush
--- a/Engine/Shaders/Private/Substrate/SubstrateRead.ush
+++ b/Engine/Shaders/Private/Substrate/SubstrateRead.ush
@@ -173,6 +173,8 @@ FSubstrateGBufferBSDF SubstrateReadGBufferBSDF(in FScreenSpaceData ScreenSpaceDa
 				bIsAtTheBottomOfTopology,
 				SharedLocalBasisIndex, SharedLocalBasisTypes,
 				true /*bSupportSSSDiffusion*/).InlinedBSDF;
+			
+			SLAB_SSSMFP(BSDF) = SanitizeSSSMFP(SLAB_SSSMFP(BSDF));
 		}
 
 		#if SUBSTRATE_INLINE_SHADING
diff --git a/Engine/Shaders/Shared/SubstrateVisualizeDefinitions.h b/Engine/Shaders/Shared/SubstrateVisualizeDefinitions.h
--- a/Engine/Shaders/Shared/SubstrateVisualizeDefinitions.h
+++ b/Engine/Shaders/Shared/SubstrateVisualizeDefinitions.h
@@ -302,25 +302,19 @@ FSubstratePixelHeader UnpackSubstrateHeaderIn(inout FSubstrateAddressing Substra
 		Out.ClosureCount = 1; 
 		Out.SetMaterialMode(Mode);
 		Out.SetHasSubsurface(bHasSSS);
+		Out.SetHasPrecShadowMask(HasPrecShadowMask(GBufferData));
+		Out.SetZeroPrecShadowMask(HasZeroPrecShadowMask(GBufferData));
+		Out.SetIsFirstPerson(IsFirstPerson(GBufferData));
+		Out.SetCastContactShadow(CastContactShadow(GBufferData));
+		Out.SetDynamicIndirectShadowCasterRepresentation(HasDynamicIndirectShadowCasterRepresentation(GBufferData));
 
-		// TODO
-		// Out.SetHasPrecShadowMask(...)
-		// Out.SetZeroPrecShadowMask(...)
-		// Out.SetIsFirstPerson(...)
-		// Out.SetCastContactShadow(...)
-		// Out.SetHasSubsurface(...)
-		// Out.SetDynamicIndirectShadowCasterRepresentation(...)
-
-		// TODO
 		#if SUBSTRATE_INLINE_SHADING
-		// Out.IrradianceAO.MaterialAO;
-		// Out.IrradianceAO.IndirectIrradiance;
-		// Out.IrradianceAO.DiffuseIndirectSampleOcclusion;
+		Out.IrradianceAO = SubstrateGetIrradianceAndAO(GBufferData);
 		#endif
 	}
 
 	SubstrateAddressing.CurrentIndex = 0;
-	SubstrateAddressing.ReadBytes = 16;
+	SubstrateAddressing.ReadBytes = 0;
 	return Out;
 }
 
@@ -679,13 +673,13 @@ FSubstratePixelDebugData ConvertToSerializableSubstratePixelDebugData(uint2 InCo
 		}
 	}
 
-	const uint FootPrint_PostBSDFs = SubstrateAddressing.ReadBytes;
+	const uint FootPrint_PostBSDFs = SUBSTRATE_GBUFFER_FORMAT == 1 ? SubstrateAddressing.ReadBytes : 16u; // 4 * 32bits MRTs
 
 	// Output memory reads
 	Data.MemoryDisplayMode = 0;
 	{
-		const bool bHasSSSData = SubstrateSubSurfaceHeaderGetIsValid(SSSHeader);
-		const uint TopLayerDataBytes = 4;
+		const bool bHasSSSData = SUBSTRATE_GBUFFER_FORMAT == 1 && SubstrateSubSurfaceHeaderGetIsValid(SSSHeader);
+		const uint TopLayerDataBytes = SUBSTRATE_GBUFFER_FORMAT == 1 ? 4 : 0;
 		const uint SSSDataBytes = 8;
 
 		const bool bSubstrateMaterial = Header.ClosureCount > 0;
