Substrate - fixed register spilling of debug shaders on some platforms.

#tests pc editor dx12, console cook
#rnx

diff --git a/Engine/Shaders/Shared/SubstrateVisualizeDefinitions.h b/Engine/Shaders/Shared/SubstrateVisualizeDefinitions.h
--- a/Engine/Shaders/Shared/SubstrateVisualizeDefinitions.h
+++ b/Engine/Shaders/Shared/SubstrateVisualizeDefinitions.h
@@ -13,8 +13,18 @@
 // Change this to force recompilation of all Substrate dependent shaders (for instance https://guidgenerator.com/online-guid-generator.aspx)
 #pragma message("UESHADERMETADATA_VERSION 24677C38-A71C-49D5-B60A-DFF59F57B39D")
 #include "/Engine/Shared/SubstrateDefinitions.h"
+
+#if PS4_PROFILE
+// For platform having trouble avoiding register spilling.
+// To fix this in general, we will have to break down the FSubstratePixelDebugData structure into Header and Closure structures with interleaved write to buffer.
+#define SUBSTRATE_VIS_CLOSURE_COUNT 1
+#else
+#define SUBSTRATE_VIS_CLOSURE_COUNT SUBSTRATE_MAX_CLOSURE_COUNT
+#endif
+
 #else
 #include "SubstrateDefinitions.h"
+#define SUBSTRATE_VIS_CLOSURE_COUNT SUBSTRATE_MAX_CLOSURE_COUNT
 #endif
 
 struct FSubsterateDebugClosure
@@ -92,7 +102,7 @@ struct FSubstratePixelDebugData
 	int HasSubsurface;
 	int LocalBasesCount;
 
-	FSubsterateDebugClosure Closures[SUBSTRATE_MAX_CLOSURE_COUNT];
+	FSubsterateDebugClosure Closures[SUBSTRATE_VIS_CLOSURE_COUNT];
 
 	int MemoryDisplayMode;
 	int MemorySlotA;
@@ -154,7 +164,7 @@ void SerializeSubstratePixelDebugData(FSubstrateDebugDataSerializer& S, FSubstra
 	S.Serialize(D.ClosureCount);
 #ifdef __cplusplus
 	// Safe guard to avoid crash in case of readback problem.
-	D.ClosureCount = FMath::Min(D.ClosureCount, int(SUBSTRATE_MAX_CLOSURE_COUNT));
+	D.ClosureCount = FMath::Min(D.ClosureCount, int(SUBSTRATE_VIS_CLOSURE_COUNT));
 #endif
 
 	S.Serialize(D.MaterialMode);
@@ -427,7 +437,7 @@ FSubstratePixelDebugData ConvertToSerializableSubstratePixelDebugData(uint2 InCo
 
 	// PSSL - Temporarely limit data serialization on PSSL due to internal shader compilation 
 	#if COMPILER_PSSL
-	SUBSTRATE_UNROLL_N(SUBSTRATE_MAX_CLOSURE_COUNT)
+	SUBSTRATE_UNROLL_N(SUBSTRATE_VIS_CLOSURE_COUNT)
 	for(uint i = 0; i < Header.ClosureCount; ++i)
 	#else
 	LOOP
