Substrate - follow up to 44777826, make sure we always clamp for both cpu and gpu world.

#tests PC
#rnx

diff --git a/Engine/Shaders/Shared/SubstrateVisualizeDefinitions.h b/Engine/Shaders/Shared/SubstrateVisualizeDefinitions.h
--- a/Engine/Shaders/Shared/SubstrateVisualizeDefinitions.h
+++ b/Engine/Shaders/Shared/SubstrateVisualizeDefinitions.h
@@ -165,6 +165,8 @@ void SerializeSubstratePixelDebugData(FSubstrateDebugDataSerializer& S, FSubstra
 #ifdef __cplusplus
 	// Safe guard to avoid crash in case of readback problem.
 	D.ClosureCount = FMath::Min(D.ClosureCount, int(SUBSTRATE_VIS_CLOSURE_COUNT));
+#else
+	D.ClosureCount = min(D.ClosureCount, int(SUBSTRATE_VIS_CLOSURE_COUNT));
 #endif
 
 	S.Serialize(D.MaterialMode);
