Substrate - follow up to 44777826, make sure we always clamp for both cpu and gpu world.

#tests PC
#rnx

diff --git a/Engine/Shaders/Shared/SubstrateVisualizeDefinitions.h b/Engine/Shaders/Shared/SubstrateVisualizeDefinitions.h
--- a/Engine/Shaders/Shared/SubstrateVisualizeDefinitions.h
+++ b/Engine/Shaders/Shared/SubstrateVisualizeDefinitions.h
@@ -414,7 +414,8 @@ FSubstratePixelDebugData ConvertToSerializableSubstratePixelDebugData(uint2 InCo
 
 	FSubstrateIrradianceAndOcclusion IrradianceAndOcclusion = SubstrateGetIrradianceAndAO(Header);
 
-	Data.ClosureCount = Header.ClosureCount;
+	const uint ClampedClosureCount = clamp(Header.ClosureCount, 0, SUBSTRATE_VIS_CLOSURE_COUNT);
+	Data.ClosureCount = ClampedClosureCount;
 
 	Data.MaterialMode = Header.GetMaterialMode();
 	Data.OptimisedLegacyMode = HEADER_MATERIALMODE_NONE;
@@ -440,10 +441,10 @@ FSubstratePixelDebugData ConvertToSerializableSubstratePixelDebugData(uint2 InCo
 	// PSSL - Temporarely limit data serialization on PSSL due to internal shader compilation 
 	#if COMPILER_PSSL
 	SUBSTRATE_UNROLL_N(SUBSTRATE_VIS_CLOSURE_COUNT)
-	for(uint i = 0; i < Header.ClosureCount; ++i)
+	for(uint i = 0; i < ClampedClosureCount; ++i)
 	#else
 	LOOP
-	for(uint i = 0; i < Header.ClosureCount; ++i)
+	for(uint i = 0; i < ClampedClosureCount; ++i)
 	#endif
 	{
 		// Unpack BSDF data
@@ -698,7 +699,7 @@ FSubstratePixelDebugData ConvertToSerializableSubstratePixelDebugData(uint2 InCo
 		const uint BSDFsSize  = FootPrint_PostBSDFs  - FootPrint_PostHeader;
 		const uint TotalSize  = (FootPrint_PostBSDFs - FootPrint_Start) + (bHasSSSData ? SSSDataBytes : 0);
 
-		if (Header.ClosureCount > 0)
+		if (ClampedClosureCount > 0)
 		{
 			if (bIsSimpleMaterial || Header.IsSingleLayerWater())
 			{
