[Nanite Skinning] Disable compressed bone transforms for now until we can track down the exact cause of jitter in dynamic wind, or we find a packed format that mitigates the issue.

diff --git a/Engine/Shaders/Shared/SkinningDefinitions.h b/Engine/Shaders/Shared/SkinningDefinitions.h
--- a/Engine/Shaders/Shared/SkinningDefinitions.h
+++ b/Engine/Shaders/Shared/SkinningDefinitions.h
@@ -61,7 +61,7 @@ struct FBankScatterHeader
 #define UINT_TYPE uint
 #endif
 
-#define USE_COMPRESSED_BONE_TRANSFORM	1
+#define USE_COMPRESSED_BONE_TRANSFORM	0
 #define USE_SKINNING_SCENE_EXTENSION_FOR_NON_NANITE	0
 
 #if USE_COMPRESSED_BONE_TRANSFORM
