[Nanite] Add a comment to explain the reason the implementation of reverse cull mode in Nanite is currently compiled out and a reference to the JIRA for bringing it back.

diff --git a/Engine/Shaders/Shared/SceneDefinitions.h b/Engine/Shaders/Shared/SceneDefinitions.h
--- a/Engine/Shaders/Shared/SceneDefinitions.h
+++ b/Engine/Shaders/Shared/SceneDefinitions.h
@@ -76,6 +76,8 @@
 #define PRIMITIVE_VISIBILITY_FLAG_HIDDEN_IN_SCENE_CAPTURE				0x200
 #define PRIMITIVE_VISIBILITY_FLAG_VISIBLE_IN_LUMEN_SCENE				0x400
 
+// NOTE: Reverse culling in Nanite was implemented, but found to create issues with HLOD/PLA that rely on the flag for certain
+// behavior with non-Nanite. Will need a resolution for these issues before we're able to add support back. See UE-174319
 #define SUPPORT_REVERSE_CULLING_IN_NANITE 0
 
 #if SUPPORT_REVERSE_CULLING_IN_NANITE
