Removed static asserts to fix cook issue.

#Synced-CL 48814992

diff --git a/Engine/Shaders/Shared/SkinningDefinitions.h b/Engine/Shaders/Shared/SkinningDefinitions.h
--- a/Engine/Shaders/Shared/SkinningDefinitions.h
+++ b/Engine/Shaders/Shared/SkinningDefinitions.h
@@ -34,8 +34,6 @@ struct FAnimBankBlockHeader
 	uint BlockTransformCount            : 16;
 };
 
-static_assert(sizeof(FAnimBankBlockHeader) == 32);
-
 #define ANIM_BANK_FLAG_NONE			0x0
 #define ANIM_BANK_FLAG_LOOPING		0x1
 #define ANIM_BANK_FLAG_AUTOSTART	0x2
@@ -78,8 +76,6 @@ struct FSkinningHeader
 	UINT_TYPE Unused                  : 6;
 };
 
-static_assert(sizeof(FSkinningHeader) == 16);
-
 #endif // defined(__cplusplus) || COMPILER_SUPPORTS_HLSL2021
 
 #if USE_COMPRESSED_BONE_TRANSFORM
