Missing skinning definition.

#Synced-CL 49999797

diff --git a/Engine/Shaders/Shared/SkinningDefinitions.h b/Engine/Shaders/Shared/SkinningDefinitions.h
--- a/Engine/Shaders/Shared/SkinningDefinitions.h
+++ b/Engine/Shaders/Shared/SkinningDefinitions.h
@@ -104,8 +104,9 @@ struct FLeafBonePathNode
 SKINNING_STATIC_ASSERT(sizeof(FLeafBonePathNode) == 4, "FLeafBonePathNode must be 4 bytes");
 struct FConcatenateBlockHeader
 {
-	uint LeafBonePathBufferOffset;
 	uint TransformBufferOffset;
+	uint LeafBonePathBufferOffset : 24;
+	uint NumPrefetchBones   : 8;
 	uint BlockLocalIndex    : 16;
 	uint BlockLeafBoneCount : 16;
 };
